Remove UE4 references in BlueprintFileUtilsBPLibrary.h

#rnx

diff --git a/Engine/Plugins/BlueprintFileUtils/Source/BlueprintFileUtils/Public/BlueprintFileUtilsBPLibrary.h b/Engine/Plugins/BlueprintFileUtils/Source/BlueprintFileUtils/Public/BlueprintFileUtilsBPLibrary.h
--- a/Engine/Plugins/BlueprintFileUtils/Source/BlueprintFileUtils/Public/BlueprintFileUtilsBPLibrary.h
+++ b/Engine/Plugins/BlueprintFileUtils/Source/BlueprintFileUtils/Public/BlueprintFileUtilsBPLibrary.h
@@ -15,7 +15,7 @@ class UBlueprintFileUtilsBPLibrary : public UBlueprintFunctionLibrary
 	/**
 	 * Finds all the files within the given directory, with optional file extension filter.
 	 *
-	 * @param Directory		The absolute path to the directory to search. Ex: "C:\UE4\Pictures"
+	 * @param Directory		The absolute path to the directory to search. Ex: "C:\UnrealEditor\Pictures"
 	 * @param FoundFiles	All the files found that matched the optional FileExtension filter, or all files if none was specified.
 	 * @param FileExtension	If FileExtension is empty string "" then all files are found.
 	 * 						Otherwise FileExtension can be of the form .EXT or just EXT and only files with that extension will be returned.
@@ -27,7 +27,7 @@ class UBlueprintFileUtilsBPLibrary : public UBlueprintFunctionLibrary
 	/**
 	 * Finds all the files and/or directories within the given directory and any sub-directories.  Files can be found with anoptional file extension filter.
 	 *
-	 * @param StartDirectory	The absolute path to the directory to start the search. Ex: "C:\UE4\Pictures"
+	 * @param StartDirectory	The absolute path to the directory to start the search. Ex: "C:\UnrealEditor\Pictures"
 	 * @param FoundPaths		All the paths (directories and/or files) found
 	 * @param Wildcard			Wildcard that can be used to find files or directories with specific text in their name.  
 								E.g *.png to find all files ending with the png extension, *images* to find anything with the word "images" in it
